NNE Fix ORTGpu ModelData backward compatibility from ORTDml

What:
When ORTGpu was introduced backward compatibility was broken for cooked data stored in the DDC cache.
-> However the data can be made backward compatible to avoid the need to reimport the related assets, this is what this CL is doing.

Test:
Generated the previous version of the asset and cook it to store it in DDC, revert the changes,
reload the project and use case and verify that the new code is able to load previous version of the cooked asset from the DDC

#rnx

diff --git a/Engine/Plugins/Experimental/NNERuntimeORTGpu/Source/NNERuntimeORT/Private/NNERuntimeORT.cpp b/Engine/Plugins/Experimental/NNERuntimeORTGpu/Source/NNERuntimeORT/Private/NNERuntimeORT.cpp
--- a/Engine/Plugins/Experimental/NNERuntimeORTGpu/Source/NNERuntimeORT/Private/NNERuntimeORT.cpp
+++ b/Engine/Plugins/Experimental/NNERuntimeORTGpu/Source/NNERuntimeORT/Private/NNERuntimeORT.cpp
@@ -80,7 +80,11 @@ bool UNNERuntimeORTGpuImpl::CanCreateModelGPU(TObjectPtr<UNNEModelData> ModelDat
 		return false;
 	}
 
+	static const FGuid DeprecatedGUID = FGuid((int32)'O', (int32)'D', (int32)'M', (int32)'L');
+
 	bool bResult = FGenericPlatformMemory::Memcmp(&(Data[0]), &(UNNERuntimeORTGpuImpl::GUID), GuidSize) == 0;
+	bResult |= FGenericPlatformMemory::Memcmp(&(Data[0]), &(DeprecatedGUID), GuidSize) == 0;
+
 	bResult &= FGenericPlatformMemory::Memcmp(&(Data[GuidSize]), &(UNNERuntimeORTGpuImpl::Version), VersionSize) == 0;
 	return bResult;
 }
