Register NNE factory and corresponding asset actions on runtime registration.

Moving duplicated code in each runtime module startup function to the common code in runtime registration. As consequence, model import and asset actions become available as soon as the first runtime is registered.

Tested by compiling and importing new assets as well as running the different asset actions.

#rnx

diff --git a/Engine/Plugins/Experimental/NNERuntimeORTCpu/Source/NNERuntimeORTCpu/NNERuntimeORTCpu.Build.cs b/Engine/Plugins/Experimental/NNERuntimeORTCpu/Source/NNERuntimeORTCpu/NNERuntimeORTCpu.Build.cs
--- a/Engine/Plugins/Experimental/NNERuntimeORTCpu/Source/NNERuntimeORTCpu/NNERuntimeORTCpu.Build.cs
+++ b/Engine/Plugins/Experimental/NNERuntimeORTCpu/Source/NNERuntimeORTCpu/NNERuntimeORTCpu.Build.cs
@@ -33,10 +33,5 @@ public class NNERuntimeORTCpu : ModuleRules
 				"NNEOnnxruntime"
 			}
 		);
-
-		if (Target.Type == TargetType.Editor)
-		{
-			PublicDependencyModuleNames.Add("NNEEditor");
-		}
 	}
 }
diff --git a/Engine/Plugins/Experimental/NNERuntimeORTCpu/Source/NNERuntimeORTCpu/Private/NNERuntimeORTCpuModule.cpp b/Engine/Plugins/Experimental/NNERuntimeORTCpu/Source/NNERuntimeORTCpu/Private/NNERuntimeORTCpuModule.cpp
--- a/Engine/Plugins/Experimental/NNERuntimeORTCpu/Source/NNERuntimeORTCpu/Private/NNERuntimeORTCpuModule.cpp
+++ b/Engine/Plugins/Experimental/NNERuntimeORTCpu/Source/NNERuntimeORTCpu/Private/NNERuntimeORTCpuModule.cpp
@@ -7,10 +7,6 @@
 
 void FNNERuntimeORTCpuModule::StartupModule()
 {
-#ifdef WITH_EDITOR
-	FModuleManager::Get().LoadModule(TEXT("NNEEditor"));
-#endif
-
 	// NNE runtime startup
 	NNERuntimeORTCpu = NewObject<UNNERuntimeORTCpuImpl>();
 	if (NNERuntimeORTCpu.IsValid())
diff --git a/Engine/Plugins/Experimental/NNERuntimeORTGpu/Source/NNERuntimeORT/NNERuntimeORT.Build.cs b/Engine/Plugins/Experimental/NNERuntimeORTGpu/Source/NNERuntimeORT/NNERuntimeORT.Build.cs
--- a/Engine/Plugins/Experimental/NNERuntimeORTGpu/Source/NNERuntimeORT/NNERuntimeORT.Build.cs
+++ b/Engine/Plugins/Experimental/NNERuntimeORTGpu/Source/NNERuntimeORT/NNERuntimeORT.Build.cs
@@ -57,10 +57,5 @@ public class NNERuntimeORT : ModuleRules
 			AddEngineThirdPartyPrivateStaticDependencies(Target, "DirectML");
 
 		}
-
-		if (Target.Type == TargetType.Editor)
-		{
-			PublicDependencyModuleNames.Add("NNEEditor");
-		}
 	}
 }
diff --git a/Engine/Plugins/Experimental/NNERuntimeORTGpu/Source/NNERuntimeORT/Private/NNERuntimeORTModule.cpp b/Engine/Plugins/Experimental/NNERuntimeORTGpu/Source/NNERuntimeORT/Private/NNERuntimeORTModule.cpp
--- a/Engine/Plugins/Experimental/NNERuntimeORTGpu/Source/NNERuntimeORT/Private/NNERuntimeORTModule.cpp
+++ b/Engine/Plugins/Experimental/NNERuntimeORTGpu/Source/NNERuntimeORT/Private/NNERuntimeORTModule.cpp
@@ -43,10 +43,6 @@ namespace UE::NNERuntimeORT::Private::DllHelper
 
 void FNNERuntimeORTModule::StartupModule()
 {
-#ifdef WITH_EDITOR
-	FModuleManager::Get().LoadModule(TEXT("NNEEditor"));
-#endif
-
 #if PLATFORM_WINDOWS
 	const FString PluginDir = IPluginManager::Get().FindPlugin("NNERuntimeORTGpu")->GetBaseDir();
 	const FString OrtBinPath = FPaths::Combine(PluginDir, TEXT(PREPROCESSOR_TO_STRING(ONNXRUNTIME_PLATFORM_PATH)));
diff --git a/Engine/Plugins/Experimental/NNERuntimeRDG/Source/NNERuntimeRDG/NNERuntimeRDG.Build.cs b/Engine/Plugins/Experimental/NNERuntimeRDG/Source/NNERuntimeRDG/NNERuntimeRDG.Build.cs
--- a/Engine/Plugins/Experimental/NNERuntimeRDG/Source/NNERuntimeRDG/NNERuntimeRDG.Build.cs
+++ b/Engine/Plugins/Experimental/NNERuntimeRDG/Source/NNERuntimeRDG/NNERuntimeRDG.Build.cs
@@ -52,10 +52,5 @@ public class NNERuntimeRDG : ModuleRules
 		{	
 			PrivateDependencyModuleNames.Add("VulkanRHI");
 		}
-
-		if (Target.Type == TargetType.Editor)
-		{
-			PublicDependencyModuleNames.Add("NNEEditor");
-		}
 	}
 }
diff --git a/Engine/Plugins/Experimental/NNERuntimeRDG/Source/NNERuntimeRDG/Private/NNERuntimeRDGModule.cpp b/Engine/Plugins/Experimental/NNERuntimeRDG/Source/NNERuntimeRDG/Private/NNERuntimeRDGModule.cpp
--- a/Engine/Plugins/Experimental/NNERuntimeRDG/Source/NNERuntimeRDG/Private/NNERuntimeRDGModule.cpp
+++ b/Engine/Plugins/Experimental/NNERuntimeRDG/Source/NNERuntimeRDG/Private/NNERuntimeRDGModule.cpp
@@ -9,10 +9,6 @@
 
 void FNNERuntimeRDGModule::StartupModule()
 {
-#ifdef WITH_EDITOR
-	FModuleManager::Get().LoadModule(TEXT("NNEEditor"));
-#endif
-
 	// NNE runtime ORT Cpu startup
 	NNERuntimeRDGHlsl = NewObject<UNNERuntimeRDGHlslImpl>();
 	if (NNERuntimeRDGHlsl.IsValid())
diff --git a/Engine/Source/Runtime/NNE/NNE.Build.cs b/Engine/Source/Runtime/NNE/NNE.Build.cs
--- a/Engine/Source/Runtime/NNE/NNE.Build.cs
+++ b/Engine/Source/Runtime/NNE/NNE.Build.cs
@@ -23,6 +23,8 @@ namespace UnrealBuildTool.Rules
 						"DerivedDataCache"
 					}
 				);
+
+				DynamicallyLoadedModuleNames.Add("NNEEditor");
 			}
 		}
     }
diff --git a/Engine/Source/Runtime/NNE/Private/NNE.cpp b/Engine/Source/Runtime/NNE/Private/NNE.cpp
--- a/Engine/Source/Runtime/NNE/Private/NNE.cpp
+++ b/Engine/Source/Runtime/NNE/Private/NNE.cpp
@@ -94,6 +94,10 @@ namespace UE::NNE
 
 	bool RegisterRuntime(TWeakInterfacePtr<INNERuntime> Runtime)
 	{
+#ifdef WITH_EDITOR
+		FModuleManager::Get().LoadModule(TEXT("NNEEditor"));
+#endif
+
 		FString RuntimeName = Runtime->GetRuntimeName();
 		FCoreDelegates::OnAllModuleLoadingPhasesComplete.AddLambda([RuntimeName]() {
 			if (FEngineAnalytics::IsAvailable())
